Add chaos user data plugin - a general way of associating data with a particle and making it accessible on the physics thread.

diff --git a/Engine/Plugins/Experimental/ChaosUserData/ChaosUserData.uplugin b/Engine/Plugins/Experimental/ChaosUserData/ChaosUserData.uplugin
new file mode 100644
--- /dev/null
+++ b/Engine/Plugins/Experimental/ChaosUserData/ChaosUserData.uplugin
@@ -0,0 +1,24 @@
+{
+	"FileVersion": 1,
+	"Version": 1,
+	"VersionName": "1.0",
+	"FriendlyName": "ChaosUserData",
+	"Description": "Custom per-particle userdata. Write-only on game thread, read-only on physics thread.",
+	"Category": "Editor",
+	"CreatedBy": "Epic Games, Inc.",
+	"CreatedByURL": "",
+	"DocsURL": "",
+	"MarketplaceURL": "",
+	"SupportURL": "",
+	"CanContainContent": false,
+	"IsBetaVersion": true,
+	"Installed": false,
+	"EnabledByDefault": true,
+	"Modules": [
+		{
+			"Name": "ChaosUserData",
+			"Type": "Runtime",
+			"LoadingPhase": "Default"
+		}
+	]
+}
\ No newline at end of file
diff --git a/Engine/Plugins/Experimental/ChaosUserData/Source/ChaosUserData/ChaosUserData.Build.cs b/Engine/Plugins/Experimental/ChaosUserData/Source/ChaosUserData/ChaosUserData.Build.cs
new file mode 100644
--- /dev/null
+++ b/Engine/Plugins/Experimental/ChaosUserData/Source/ChaosUserData/ChaosUserData.Build.cs
@@ -0,0 +1,24 @@
+// Copyright Epic Games, Inc. All Rights Reserved.
+
+using UnrealBuildTool;
+
+namespace UnrealBuildTool.Rules
+{
+	public class ChaosUserData : ModuleRules
+	{
+
+		public ChaosUserData(ReadOnlyTargetRules Target) : base(Target)
+		{
+			PublicDependencyModuleNames.AddRange(
+				new string[]
+				{
+					"Core",
+					"CoreUObject"
+				}
+			);
+
+			SetupModulePhysicsSupport(Target);
+			PrivateDefinitions.Add("CHAOS_INCLUDE_LEVEL_1=1");
+		}
+	}
+}
diff --git a/Engine/Plugins/Experimental/ChaosUserData/Source/ChaosUserData/Private/ChaosUserDataModule.cpp b/Engine/Plugins/Experimental/ChaosUserData/Source/ChaosUserData/Private/ChaosUserDataModule.cpp
new file mode 100644
--- /dev/null
+++ b/Engine/Plugins/Experimental/ChaosUserData/Source/ChaosUserData/Private/ChaosUserDataModule.cpp
@@ -0,0 +1,16 @@
+// Copyright Epic Games, Inc. All Rights Reserved.
+
+#include "ChaosUserDataModule.h"
+
+void FChaosUserDataModule::StartupModule()
+{
+	// This code will execute after your module is loaded into memory; the exact timing is specified in the .uplugin file per-module
+}
+
+void FChaosUserDataModule::ShutdownModule()
+{
+	// This function may be called during shutdown to clean up your module.  For modules that support dynamic reloading,
+	// we call this function before unloading the module.
+}
+
+IMPLEMENT_MODULE(FChaosUserDataModule, ChaosUserDataModule)
\ No newline at end of file
diff --git a/Engine/Plugins/Experimental/ChaosUserData/Source/ChaosUserData/Private/ChaosUserDataStats.cpp b/Engine/Plugins/Experimental/ChaosUserData/Source/ChaosUserData/Private/ChaosUserDataStats.cpp
new file mode 100644
--- /dev/null
+++ b/Engine/Plugins/Experimental/ChaosUserData/Source/ChaosUserData/Private/ChaosUserDataStats.cpp
@@ -0,0 +1,10 @@
+#include "ChaosUserDataStats.h"
+
+namespace Chaos
+{
+	DEFINE_STAT(STAT_SetData_External);
+	DEFINE_STAT(STAT_RemoveData_External);
+	DEFINE_STAT(STAT_GetData_Internal);
+	DEFINE_STAT(STAT_Tick_UpdateData);
+	DEFINE_STAT(STAT_Tick_RemoveData);
+}
diff --git a/Engine/Plugins/Experimental/ChaosUserData/Source/ChaosUserData/Public/ChaosUserData.h b/Engine/Plugins/Experimental/ChaosUserData/Source/ChaosUserData/Public/ChaosUserData.h
new file mode 100644
--- /dev/null
+++ b/Engine/Plugins/Experimental/ChaosUserData/Source/ChaosUserData/Public/ChaosUserData.h
@@ -0,0 +1,207 @@
+// Copyright Epic Games, Inc. All Rights Reserved.
+
+#pragma once
+
+#include "CoreMinimal.h"
+#include "Chaos/PBDRigidsEvolutionFwd.h"
+#include "Chaos/SimCallbackInput.h"
+#include "Chaos/SimCallbackObject.h"
+#include "Chaos/GeometryParticlesfwd.h"
+#include "Chaos/Framework/PhysicsSolverBase.h"
+#include "UObject/WeakObjectPtr.h"
+#include "PhysicsInterfaceDeclaresCore.h"
+#include "PhysicsProxy/SingleParticlePhysicsProxy.h"
+#include "Templates/UniquePtr.h"
+#include "ChaosUserDataStats.h"
+
+/*
+
+	Chaos User Data
+	===============
+
+	The idea behind this tool is to provide a generic way of associating custom data
+	with physics particles, which is write-only from the game thread, and read-only
+	from the physics thread.
+
+	This comes in handy when physical interactions at the per-contact level need to
+	be affected by gameplay properties.
+
+	In order to use a Chaos::TUserDataManager it will need to be created using the chaos
+	solver's Chaos::FPhysicsSolverBase::CreateAndRegisterSimCallbackObject_External.
+	This library does not natively provide a method of accessing the appropriate
+	Chaos::TUserDataManager from the physics thread, but this can be achieved in a number
+	of ways - it is left up to the game to decide how to do this for flexibility.
+
+*/
+
+namespace Chaos
+{
+	// TUserDataManagerInput
+	//
+	// Input is a collection of pointers to new and updated userdata objects to be sent
+	// to the physics thread
+	template <typename TUserData>
+	struct TUserDataManagerInput : public Chaos::FSimCallbackInput
+	{
+		// Map of particle unique indices to user data ptrs
+		// 
+		// NOTE: This is marked mutable because the userdata objects must be moved
+		// to the internal array after making it OnPreSimulate_Internal, but input
+		// objects are const in that context. This might be frowned on, but since
+		// TUserDataManagerInput is a class which is only used internally to
+		// TUserDataManager, an argument could be made either way.
+		mutable TMap<Chaos::FUniqueIdx, TUniquePtr<TUserData>> UserDataToAdd;
+
+		// Set of particle unique indices for which to remove user data
+		TSet<Chaos::FUniqueIdx> UserDataToRemove;
+
+		// Monotonically increasing identifier for the input object. Each newly
+		// constructed input will store and increment this;
+		int32 Identifier = -1;
+
+		void Reset()
+		{
+			UserDataToAdd.Reset();
+			UserDataToRemove.Reset();
+			Identifier = -1;
+		}
+	};
+
+	// TUserDataManager
+	//
+	// A chaos callback object which stores and allows access to user data associated with
+	// particles on the physics thread.
+	//
+	// Note that Chaos::FSimCallbackOutput is the output struct - this carries no data because
+	// this is a one-way callback. We use it basically just to marshal data in one direction.
+	template <typename TUserData>
+	class TUserDataManager : public Chaos::TSimCallbackObject<TUserDataManagerInput<TUserData>, Chaos::FSimCallbackNoOutput>
+	{
+		using TInput = TUserDataManagerInput<TUserData>;
+
+	public:
+
+		virtual ~TUserDataManager() { }
+
+		// Add or update user data associated with this particle handle
+		bool SetData_External(const Chaos::FRigidBodyHandle_External& Handle, const TUserData& UserData)
+		{
+			SCOPE_CYCLE_COUNTER(STAT_SetData_External);
+
+			if (const Chaos::FPhysicsSolverBase* Solver = this->GetSolver())
+			{
+				if (TInput* Input = this->GetProducerInputData_External())
+				{
+					// Add the data to the map to be sent to physics thread
+					Input->UserDataToAdd.Emplace(Handle.UniqueIdx(), MakeUnique<TUserData>(UserData));
+
+					// In case it was removed and then added again in the same frame,
+					// untrack this particle for data removal
+					Input->UserDataToRemove.Remove(Handle.UniqueIdx());
+
+					// If this is a new input, set it's identifier
+					if (Input->Identifier == -1)
+					{
+						Input->Identifier = InputIdentifier_External++;
+					}
+
+					// Successfully queued for add/update
+					return true;
+				}
+			}
+
+			// Failed to queue for add/update
+			return false;
+		}
+
+		// Remove user data associated with this particle handle
+		bool RemoveData_External(const Chaos::FRigidBodyHandle_External& Handle)
+		{
+			SCOPE_CYCLE_COUNTER(STAT_RemoveData_External);
+
+			if (this->GetSolver() != nullptr)
+			{
+				if (TInput* Input = this->GetProducerInputData_External())
+				{
+					// Track the particle for removal
+					Input->UserDataToRemove.Add(Handle.UniqueIdx());
+
+					// In case it was added/updated and then removed in the same frame,
+					// untrack the add/update
+					Input->UserDataToAdd.Remove(Handle.UniqueIdx());
+
+					// Successfully queued for removal
+					return true;
+				}
+			}
+
+			// Failed to queue for removal
+			return false;
+		}
+
+		// TParticleHandle is generalized here because it can be Chaos::FRigidBodyHandle_Internal
+		// or Chaos::FGeometryParticleHandle which have the same api...
+		template <typename TParticleHandle>
+		const TUserData* GetData_Internal(const TParticleHandle& Handle) const
+		{
+			SCOPE_CYCLE_COUNTER(STAT_GetData_Internal);
+
+			const int32 Idx = Handle.UniqueIdx().Idx;
+			return UserDataMap_Internal.IsValidIndex(Idx) ? UserDataMap_Internal[Idx].Get() : nullptr;
+		}
+
+
+	protected:
+
+		virtual void OnPreSimulate_Internal() override
+		{
+			if (const TInput* Input = this->GetConsumerInput_Internal())
+			{
+				// Only proceed if the input has not yet been processed.
+				// 
+				// It's possible that we'll get multiple presimulate calls with
+				// the same input because the same input continues to be provided
+				// until a new one is received, so we cache the timestamp of the
+				// last processed input to make sure that we don't double-process it.
+				if (InputIdentifier_Internal != Input->Identifier)
+				{
+					InputIdentifier_Internal = Input->Identifier;
+
+					if (Input->UserDataToAdd.Num() > 0)
+					{
+						SCOPE_CYCLE_COUNTER(STAT_Tick_UpdateData);
+
+						// Move all the user data to the internal map
+						for (auto& Iter : Input->UserDataToAdd)
+						{
+							UserDataMap_Internal.EmplaceAt(Iter.Key.Idx, MoveTemp(Iter.Value));
+						}
+					}
+
+					if (Input->UserDataToRemove.Num() > 0)
+					{
+						SCOPE_CYCLE_COUNTER(STAT_Tick_RemoveData);
+
+						// Delete user data that has been removed
+						for (const Chaos::FUniqueIdx Idx : Input->UserDataToRemove)
+						{
+							UserDataMap_Internal.RemoveAt(Idx.Idx);
+						}
+
+						// Shrink sparse array if we took elements off the end
+						UserDataMap_Internal.Shrink();
+					}
+				}
+			}
+		}
+
+		// Identifier of the next input to be created on the external thread
+		int32 InputIdentifier_External = 0;
+
+		// Identifier of the last input to be consumed on the internal thread
+		int32 InputIdentifier_Internal = -1;
+
+		// Map of particle unique ids to user data
+		TSparseArray<TUniquePtr<TUserData>> UserDataMap_Internal;
+	};
+}
diff --git a/Engine/Plugins/Experimental/ChaosUserData/Source/ChaosUserData/Public/ChaosUserDataModule.h b/Engine/Plugins/Experimental/ChaosUserData/Source/ChaosUserData/Public/ChaosUserDataModule.h
new file mode 100644
--- /dev/null
+++ b/Engine/Plugins/Experimental/ChaosUserData/Source/ChaosUserData/Public/ChaosUserDataModule.h
@@ -0,0 +1,15 @@
+// Copyright Epic Games, Inc. All Rights Reserved.
+
+#pragma once
+
+#include "CoreMinimal.h"
+#include "Modules/ModuleManager.h"
+
+class FChaosUserDataModule : public IModuleInterface
+{
+public:
+
+	/** IModuleInterface implementation */
+	virtual void StartupModule() override;
+	virtual void ShutdownModule() override;
+};
\ No newline at end of file
diff --git a/Engine/Plugins/Experimental/ChaosUserData/Source/ChaosUserData/Public/ChaosUserDataStats.h b/Engine/Plugins/Experimental/ChaosUserData/Source/ChaosUserData/Public/ChaosUserDataStats.h
new file mode 100644
--- /dev/null
+++ b/Engine/Plugins/Experimental/ChaosUserData/Source/ChaosUserData/Public/ChaosUserDataStats.h
@@ -0,0 +1,17 @@
+// Copyright Epic Games, Inc. All Rights Reserved.
+
+#pragma once
+
+#include "Stats/Stats.h"
+
+DECLARE_STATS_GROUP(TEXT("ChaosUserData"), STATGROUP_ChaosUserData, STATCAT_Advanced);
+
+namespace Chaos
+{
+	DECLARE_CYCLE_STAT_EXTERN(TEXT("ChaosUserData::SetData_External"), STAT_SetData_External, STATGROUP_ChaosUserData, CHAOSUSERDATA_API);
+	DECLARE_CYCLE_STAT_EXTERN(TEXT("ChaosUserData::RemoveData_External"), STAT_RemoveData_External, STATGROUP_ChaosUserData, CHAOSUSERDATA_API);
+	DECLARE_CYCLE_STAT_EXTERN(TEXT("ChaosUserData::GetData_Internal"), STAT_GetData_Internal, STATGROUP_ChaosUserData, CHAOSUSERDATA_API);
+	DECLARE_CYCLE_STAT_EXTERN(TEXT("ChaosUserData::OnPreSimulate_Internal::UpdateData"), STAT_Tick_UpdateData, STATGROUP_ChaosUserData, CHAOSUSERDATA_API);
+	DECLARE_CYCLE_STAT_EXTERN(TEXT("ChaosUserData::OnPreSimulate_Internal::RemoveData"), STAT_Tick_RemoveData, STATGROUP_ChaosUserData, CHAOSUSERDATA_API);
+}
+
