Add missing copyright boilerplate

diff --git a/Engine/Plugins/Experimental/ChaosUserData/Source/ChaosUserData/Private/ChaosUserDataStats.cpp b/Engine/Plugins/Experimental/ChaosUserData/Source/ChaosUserData/Private/ChaosUserDataStats.cpp
--- a/Engine/Plugins/Experimental/ChaosUserData/Source/ChaosUserData/Private/ChaosUserDataStats.cpp
+++ b/Engine/Plugins/Experimental/ChaosUserData/Source/ChaosUserData/Private/ChaosUserDataStats.cpp
@@ -1,3 +1,5 @@
+// Copyright Epic Games, Inc. All Rights Reserved.
+
 #include "ChaosUserDataStats.h"
 
 namespace Chaos
