Fix IMPLEMENT_MODULE line in ChaosUserDataModule.cpp

diff --git a/Engine/Plugins/Experimental/ChaosUserData/Source/ChaosUserData/Private/ChaosUserDataModule.cpp b/Engine/Plugins/Experimental/ChaosUserData/Source/ChaosUserData/Private/ChaosUserDataModule.cpp
--- a/Engine/Plugins/Experimental/ChaosUserData/Source/ChaosUserData/Private/ChaosUserDataModule.cpp
+++ b/Engine/Plugins/Experimental/ChaosUserData/Source/ChaosUserData/Private/ChaosUserDataModule.cpp
@@ -13,4 +13,4 @@ void FChaosUserDataModule::ShutdownModule()
 	// we call this function before unloading the module.
 }
 
-IMPLEMENT_MODULE(FChaosUserDataModule, ChaosUserDataModule)
\ No newline at end of file
+IMPLEMENT_MODULE(FChaosUserDataModule, ChaosUserData)
\ No newline at end of file
